Added a new BTDecorator_GameplayTagQuery that uses gameplay events on certain gameplay tags to validate or not a gameplay query

diff --git a/Engine/Plugins/Experimental/GameplayBehaviors/Source/GameplayBehaviorsModule/Private/AI/BTDecorator_GameplayTagQuery.cpp b/Engine/Plugins/Experimental/GameplayBehaviors/Source/GameplayBehaviorsModule/Private/AI/BTDecorator_GameplayTagQuery.cpp
new file mode 100644
--- /dev/null
+++ b/Engine/Plugins/Experimental/GameplayBehaviors/Source/GameplayBehaviorsModule/Private/AI/BTDecorator_GameplayTagQuery.cpp
@@ -0,0 +1,144 @@
+// Copyright Epic Games, Inc. All Rights Reserved.
+
+#include "AI/BTDecorator_GameplayTagQuery.h"
+#include "AbilitySystemComponent.h"
+#include "AbilitySystemGlobals.h"
+#include "BehaviorTree/BehaviorTreeComponent.h"
+#include "BehaviorTree/BlackboardComponent.h"
+#include "BehaviorTree/Blackboard/BlackboardKeyType_Object.h"
+#include "GameFramework/Actor.h"
+#include "GameplayTagAssetInterface.h"
+#include "UObject/Object.h"
+
+UBTDecorator_GameplayTagQuery::UBTDecorator_GameplayTagQuery(const FObjectInitializer& ObjectInitializer) : Super(ObjectInitializer)
+{
+	NodeName = "Gameplay Tag Query";
+
+	INIT_DECORATOR_NODE_NOTIFY_FLAGS();
+
+	// Accept only actors
+	ActorForGameplayTagQuery.AddObjectFilter(this, GET_MEMBER_NAME_CHECKED(UBTDecorator_GameplayTagQuery, ActorForGameplayTagQuery), AActor::StaticClass());
+
+	// Default to using Self Actor
+	ActorForGameplayTagQuery.SelectedKeyName = FBlackboard::KeySelf;
+}
+
+bool UBTDecorator_GameplayTagQuery::CalculateRawConditionValue(UBehaviorTreeComponent& OwnerComp, uint8* NodeMemory) const
+{
+	const UBlackboardComponent* BlackboardComp = OwnerComp.GetBlackboardComponent();
+	if (!BlackboardComp)
+	{
+		// Not calling super here since it returns true
+		return false;
+	}
+
+	const IGameplayTagAssetInterface* GameplayTagAssetInterface = Cast<IGameplayTagAssetInterface>(BlackboardComp->GetValue<UBlackboardKeyType_Object>(ActorForGameplayTagQuery.GetSelectedKeyID()));
+	if (!GameplayTagAssetInterface)
+	{
+		// Not calling super here since it returns true
+		return false;
+	}
+
+	FGameplayTagContainer SelectedActorTags;
+	GameplayTagAssetInterface->GetOwnedGameplayTags(SelectedActorTags);
+
+	return GameplayTagQuery.Matches(SelectedActorTags);
+}
+
+void UBTDecorator_GameplayTagQuery::OnGameplayTagInQueryChanged(const FGameplayTag InTag, int32 NewCount, TWeakObjectPtr<UBehaviorTreeComponent> BehaviorTreeComponent, uint8* NodeMemory)
+{
+	if (!BehaviorTreeComponent.IsValid())
+	{
+		return;
+	}
+
+	ConditionalFlowAbort(*BehaviorTreeComponent, EBTDecoratorAbortRequest::ConditionResultChanged);
+}
+
+FString UBTDecorator_GameplayTagQuery::GetStaticDescription() const
+{
+	return FString::Printf(TEXT("%s: %s"), *Super::GetStaticDescription(), *GameplayTagQuery.GetDescription());
+}
+
+void UBTDecorator_GameplayTagQuery::CleanupMemory(UBehaviorTreeComponent& OwnerComp, uint8* NodeMemory, EBTMemoryClear::Type CleanupType) const
+{
+	ensureMsgf(GameplayTagEventHandles.Num() == 0, TEXT("Dangling gameplay tag event handles for decorator %s"), *GetStaticDescription());
+}
+
+void UBTDecorator_GameplayTagQuery::OnBecomeRelevant(UBehaviorTreeComponent& OwnerComp, uint8* NodeMemory)
+{
+	UBlackboardComponent* BlackboardComp = OwnerComp.GetBlackboardComponent();
+	if (!BlackboardComp)
+	{
+		// Not calling super here since it does nothing
+		return;
+	}
+
+	const AActor* SelectedActor = Cast<AActor>(BlackboardComp->GetValue<UBlackboardKeyType_Object>(ActorForGameplayTagQuery.GetSelectedKeyID()));
+	if (!SelectedActor)
+	{
+		// Not calling super here since it does nothing
+		return;
+	}
+
+	CachedAbilitySystemComponent = UAbilitySystemGlobals::GetAbilitySystemComponentFromActor(SelectedActor);
+	if (CachedAbilitySystemComponent)
+	{
+		for (const FGameplayTag& CurrentTag : QueryTags)
+		{
+			FDelegateHandle GameplayTagEventCallbackDelegate = CachedAbilitySystemComponent->RegisterGameplayTagEvent(CurrentTag, EGameplayTagEventType::Type::AnyCountChange).AddUObject(this, &UBTDecorator_GameplayTagQuery::OnGameplayTagInQueryChanged, TWeakObjectPtr<UBehaviorTreeComponent>(&OwnerComp), NodeMemory);
+			GameplayTagEventHandles.Emplace(CurrentTag, GameplayTagEventCallbackDelegate);
+		}
+	}
+}
+
+void UBTDecorator_GameplayTagQuery::OnCeaseRelevant(UBehaviorTreeComponent& OwnerComp, uint8* NodeMemory)
+{
+	const UBlackboardComponent* BlackboardComp = OwnerComp.GetBlackboardComponent();
+	if (!BlackboardComp)
+	{
+		// Not calling super here since it does nothing
+		return;
+	}
+
+	if (CachedAbilitySystemComponent)
+	{
+		for (const TTuple<FGameplayTag, FDelegateHandle>& GameplayTagEvent : GameplayTagEventHandles)
+		{
+			CachedAbilitySystemComponent->RegisterGameplayTagEvent(GameplayTagEvent.Key, EGameplayTagEventType::Type::AnyCountChange).Remove(GameplayTagEvent.Value);
+		}
+	}
+
+	GameplayTagEventHandles.Reset();
+	CachedAbilitySystemComponent = nullptr;
+}
+
+#if WITH_EDITOR
+void UBTDecorator_GameplayTagQuery::CacheGameplayTagsInsideQuery()
+{
+	QueryTags.Reset();
+	GameplayTagQuery.GetGameplayTagArray(QueryTags);
+}
+
+void UBTDecorator_GameplayTagQuery::PostEditChangeProperty(FPropertyChangedEvent& PropertyChangedEvent)
+{
+	Super::PostEditChangeProperty(PropertyChangedEvent);
+	if (PropertyChangedEvent.Property == NULL)
+	{
+		return;
+	}
+
+	CacheGameplayTagsInsideQuery();
+}
+#endif	// WITH_EDITOR
+
+void UBTDecorator_GameplayTagQuery::InitializeFromAsset(UBehaviorTree& Asset)
+{
+	Super::InitializeFromAsset(Asset);
+
+	UBlackboardData* BBAsset = GetBlackboardAsset();
+	if (ensure(BBAsset))
+	{
+		ActorForGameplayTagQuery.ResolveSelectedKey(*BBAsset);
+	}
+}
diff --git a/Engine/Plugins/Experimental/GameplayBehaviors/Source/GameplayBehaviorsModule/Public/AI/BTDecorator_GameplayTagQuery.h b/Engine/Plugins/Experimental/GameplayBehaviors/Source/GameplayBehaviorsModule/Public/AI/BTDecorator_GameplayTagQuery.h
new file mode 100644
--- /dev/null
+++ b/Engine/Plugins/Experimental/GameplayBehaviors/Source/GameplayBehaviorsModule/Public/AI/BTDecorator_GameplayTagQuery.h
@@ -0,0 +1,76 @@
+// Copyright Epic Games, Inc. All Rights Reserved.
+#pragma once
+
+#include "CoreMinimal.h"
+#include "UObject/ObjectMacros.h"
+#include "GameplayTagContainer.h"
+#include "BehaviorTree/BTDecorator.h"
+#include "BehaviorTree/BehaviorTreeTypes.h"
+#include "Delegates/IDelegateInstance.h"
+#include "UObject/WeakObjectPtrTemplates.h"
+#include "Templates/Tuple.h"
+
+#include "BTDecorator_GameplayTagQuery.generated.h"
+
+class UBehaviorTree;
+class UBehaviorTreeComponent;
+class UBlackboardComponent;
+class UAbilitySystemComponent;
+
+/**
+ * GameplayTagQuery decorator node.
+ * A decorator node that bases its condition on matching a gameplay tag query.
+ */
+
+UCLASS(HideCategories=(Condition))
+class GAMEPLAYBEHAVIORSMODULE_API UBTDecorator_GameplayTagQuery : public UBTDecorator
+{
+	GENERATED_BODY()
+	UBTDecorator_GameplayTagQuery(const FObjectInitializer& ObjectInitializer);
+
+	virtual bool CalculateRawConditionValue(UBehaviorTreeComponent& OwnerComp, uint8* NodeMemory) const override;
+
+	/** Callback for when a tag in our query changes */
+	void OnGameplayTagInQueryChanged(const FGameplayTag InTag, int32 NewCount, TWeakObjectPtr<UBehaviorTreeComponent> BehaviorTreeComponent, uint8* NodeMemory);
+
+	virtual FString GetStaticDescription() const override;
+
+	virtual void CleanupMemory(UBehaviorTreeComponent& OwnerComp, uint8* NodeMemory, EBTMemoryClear::Type CleanupType) const override;
+
+protected:
+
+	UPROPERTY(EditAnywhere, Category=GameplayTagQuery, Meta=(ToolTips="Which Actor (from the blackboard) should be checked for this gameplay tag query?"))
+	FBlackboardKeySelector ActorForGameplayTagQuery;
+
+	/** Gameplay tag query to match */
+	UPROPERTY(EditAnywhere, Category=GameplayTagQuery)
+	FGameplayTagQuery GameplayTagQuery;
+
+	UPROPERTY()
+	TArray<FGameplayTag> QueryTags;
+
+	UPROPERTY(Transient)
+	UAbilitySystemComponent* CachedAbilitySystemComponent = nullptr;
+	
+	/** Array of handles for our gameplay tag query delegates */
+	TArray<TTuple<FGameplayTag, FDelegateHandle>> GameplayTagEventHandles;
+
+	/** called when execution flow controller becomes active */
+	virtual void OnBecomeRelevant(UBehaviorTreeComponent& OwnerComp, uint8* NodeMemory) override;
+
+	/** called when execution flow controller becomes inactive */
+	virtual void OnCeaseRelevant(UBehaviorTreeComponent& OwnerComp, uint8* NodeMemory) override;
+
+#if WITH_EDITOR
+	/** Get the array of tags onto which we need to add delegates 
+		The gameplay tag interface forces us to allocate everytime we 
+		want to know the list of gameplay tags inside a query so we 
+		do this only once and cache it.
+	*/
+	virtual void CacheGameplayTagsInsideQuery();
+
+	virtual void PostEditChangeProperty(FPropertyChangedEvent& PropertyChangedEvent) override;
+#endif
+
+	virtual void InitializeFromAsset(UBehaviorTree& Asset) override;
+};
