Fix Error: Native pointer usage in member declaration detected
#rnx

diff --git a/Engine/Plugins/Experimental/GameplayBehaviors/Source/GameplayBehaviorsModule/Public/AI/BTDecorator_GameplayTagQuery.h b/Engine/Plugins/Experimental/GameplayBehaviors/Source/GameplayBehaviorsModule/Public/AI/BTDecorator_GameplayTagQuery.h
--- a/Engine/Plugins/Experimental/GameplayBehaviors/Source/GameplayBehaviorsModule/Public/AI/BTDecorator_GameplayTagQuery.h
+++ b/Engine/Plugins/Experimental/GameplayBehaviors/Source/GameplayBehaviorsModule/Public/AI/BTDecorator_GameplayTagQuery.h
@@ -50,7 +50,7 @@ protected:
 	TArray<FGameplayTag> QueryTags;
 
 	UPROPERTY(Transient)
-	UAbilitySystemComponent* CachedAbilitySystemComponent = nullptr;
+	TObjectPtr<UAbilitySystemComponent> CachedAbilitySystemComponent;
 	
 	/** Array of handles for our gameplay tag query delegates */
 	TArray<TTuple<FGameplayTag, FDelegateHandle>> GameplayTagEventHandles;
