Moved runtime data to decorator memory since the node is not instanced

diff --git a/Engine/Plugins/Experimental/GameplayBehaviors/Source/GameplayBehaviorsModule/Private/AI/BTDecorator_GameplayTagQuery.cpp b/Engine/Plugins/Experimental/GameplayBehaviors/Source/GameplayBehaviorsModule/Private/AI/BTDecorator_GameplayTagQuery.cpp
--- a/Engine/Plugins/Experimental/GameplayBehaviors/Source/GameplayBehaviorsModule/Private/AI/BTDecorator_GameplayTagQuery.cpp
+++ b/Engine/Plugins/Experimental/GameplayBehaviors/Source/GameplayBehaviorsModule/Private/AI/BTDecorator_GameplayTagQuery.cpp
@@ -62,7 +62,8 @@ FString UBTDecorator_GameplayTagQuery::GetStaticDescription() const
 
 void UBTDecorator_GameplayTagQuery::CleanupMemory(UBehaviorTreeComponent& OwnerComp, uint8* NodeMemory, EBTMemoryClear::Type CleanupType) const
 {
-	ensureMsgf(GameplayTagEventHandles.Num() == 0, TEXT("Dangling gameplay tag event handles for decorator %s"), *GetStaticDescription());
+	const FBTDecorator_GameplayTagQueryMemory* MyMemory = CastInstanceNodeMemory<FBTDecorator_GameplayTagQueryMemory>(NodeMemory);
+	ensureMsgf(MyMemory->GameplayTagEventHandles.Num() == 0, TEXT("Dangling gameplay tag event handles for decorator %s"), *GetStaticDescription());
 }
 
 void UBTDecorator_GameplayTagQuery::OnBecomeRelevant(UBehaviorTreeComponent& OwnerComp, uint8* NodeMemory)
@@ -81,36 +82,38 @@ void UBTDecorator_GameplayTagQuery::OnBecomeRelevant(UBehaviorTreeComponent& Own
 		return;
 	}
 
-	CachedAbilitySystemComponent = UAbilitySystemGlobals::GetAbilitySystemComponentFromActor(SelectedActor);
-	if (CachedAbilitySystemComponent)
+	FBTDecorator_GameplayTagQueryMemory* MyMemory = CastInstanceNodeMemory<FBTDecorator_GameplayTagQueryMemory>(NodeMemory);
+	MyMemory->CachedAbilitySystemComponent = UAbilitySystemGlobals::GetAbilitySystemComponentFromActor(SelectedActor);
+
+	if (MyMemory->CachedAbilitySystemComponent.IsValid())
 	{
 		for (const FGameplayTag& CurrentTag : QueryTags)
 		{
-			FDelegateHandle GameplayTagEventCallbackDelegate = CachedAbilitySystemComponent->RegisterGameplayTagEvent(CurrentTag, EGameplayTagEventType::Type::AnyCountChange).AddUObject(this, &UBTDecorator_GameplayTagQuery::OnGameplayTagInQueryChanged, TWeakObjectPtr<UBehaviorTreeComponent>(&OwnerComp), NodeMemory);
-			GameplayTagEventHandles.Emplace(CurrentTag, GameplayTagEventCallbackDelegate);
+			FDelegateHandle GameplayTagEventCallbackDelegate = MyMemory->CachedAbilitySystemComponent.Get()->RegisterGameplayTagEvent(CurrentTag, EGameplayTagEventType::Type::AnyCountChange).AddUObject(this, &UBTDecorator_GameplayTagQuery::OnGameplayTagInQueryChanged, TWeakObjectPtr<UBehaviorTreeComponent>(&OwnerComp), NodeMemory);
+			MyMemory->GameplayTagEventHandles.Emplace(CurrentTag, GameplayTagEventCallbackDelegate);
 		}
 	}
 }
 
 void UBTDecorator_GameplayTagQuery::OnCeaseRelevant(UBehaviorTreeComponent& OwnerComp, uint8* NodeMemory)
 {
-	const UBlackboardComponent* BlackboardComp = OwnerComp.GetBlackboardComponent();
-	if (!BlackboardComp)
-	{
-		// Not calling super here since it does nothing
-		return;
-	}
+	FBTDecorator_GameplayTagQueryMemory* MyMemory = CastInstanceNodeMemory<FBTDecorator_GameplayTagQueryMemory>(NodeMemory);
 
-	if (CachedAbilitySystemComponent)
+	if (MyMemory->CachedAbilitySystemComponent.IsValid())
 	{
-		for (const TTuple<FGameplayTag, FDelegateHandle>& GameplayTagEvent : GameplayTagEventHandles)
+		for (const TTuple<FGameplayTag, FDelegateHandle>& GameplayTagEvent : MyMemory->GameplayTagEventHandles)
 		{
-			CachedAbilitySystemComponent->RegisterGameplayTagEvent(GameplayTagEvent.Key, EGameplayTagEventType::Type::AnyCountChange).Remove(GameplayTagEvent.Value);
+			MyMemory->CachedAbilitySystemComponent.Get()->RegisterGameplayTagEvent(GameplayTagEvent.Key, EGameplayTagEventType::Type::AnyCountChange).Remove(GameplayTagEvent.Value);
 		}
 	}
 
-	GameplayTagEventHandles.Reset();
-	CachedAbilitySystemComponent = nullptr;
+	MyMemory->GameplayTagEventHandles.Reset();
+	MyMemory->CachedAbilitySystemComponent = nullptr;
+}
+
+uint16 UBTDecorator_GameplayTagQuery::GetInstanceMemorySize() const
+{
+	return sizeof(FBTDecorator_GameplayTagQueryMemory);
 }
 
 #if WITH_EDITOR
diff --git a/Engine/Plugins/Experimental/GameplayBehaviors/Source/GameplayBehaviorsModule/Public/AI/BTDecorator_GameplayTagQuery.h b/Engine/Plugins/Experimental/GameplayBehaviors/Source/GameplayBehaviorsModule/Public/AI/BTDecorator_GameplayTagQuery.h
--- a/Engine/Plugins/Experimental/GameplayBehaviors/Source/GameplayBehaviorsModule/Public/AI/BTDecorator_GameplayTagQuery.h
+++ b/Engine/Plugins/Experimental/GameplayBehaviors/Source/GameplayBehaviorsModule/Public/AI/BTDecorator_GameplayTagQuery.h
@@ -17,6 +17,14 @@ class UBehaviorTreeComponent;
 class UBlackboardComponent;
 class UAbilitySystemComponent;
 
+struct FBTDecorator_GameplayTagQueryMemory
+{
+	TWeakObjectPtr<UAbilitySystemComponent> CachedAbilitySystemComponent;
+	
+	/** Array of handles for our gameplay tag query delegates */
+	TArray<TTuple<FGameplayTag, FDelegateHandle>> GameplayTagEventHandles;
+};
+
 /**
  * GameplayTagQuery decorator node.
  * A decorator node that bases its condition on matching a gameplay tag query.
@@ -49,18 +57,14 @@ protected:
 	UPROPERTY()
 	TArray<FGameplayTag> QueryTags;
 
-	UPROPERTY(Transient)
-	TObjectPtr<UAbilitySystemComponent> CachedAbilitySystemComponent;
-	
-	/** Array of handles for our gameplay tag query delegates */
-	TArray<TTuple<FGameplayTag, FDelegateHandle>> GameplayTagEventHandles;
-
 	/** called when execution flow controller becomes active */
 	virtual void OnBecomeRelevant(UBehaviorTreeComponent& OwnerComp, uint8* NodeMemory) override;
 
 	/** called when execution flow controller becomes inactive */
 	virtual void OnCeaseRelevant(UBehaviorTreeComponent& OwnerComp, uint8* NodeMemory) override;
 
+	virtual uint16 GetInstanceMemorySize() const override;
+
 #if WITH_EDITOR
 	/** Get the array of tags onto which we need to add delegates 
 		The gameplay tag interface forces us to allocate everytime we 
