Added a gameplay tag blackboard key.
Also added none as a selectable option when nones are allowed for blackboard key selectors.

diff --git a/Engine/Plugins/Experimental/GameplayBehaviors/GameplayBehaviors.uplugin b/Engine/Plugins/Experimental/GameplayBehaviors/GameplayBehaviors.uplugin
--- a/Engine/Plugins/Experimental/GameplayBehaviors/GameplayBehaviors.uplugin
+++ b/Engine/Plugins/Experimental/GameplayBehaviors/GameplayBehaviors.uplugin
@@ -14,8 +14,7 @@
 	"CanContainContent" : false,
 	"IsBetaVersion" : true,
 	"Installed" : false,
-	"Modules":
-	[
+	"Modules": [
 		{
 			"Name": "GameplayBehaviorsModule",
 			"Type": "Runtime",
@@ -25,6 +24,11 @@
 			"Name": "GameplayBehaviorsModule",
 			"Type": "UncookedOnly",
 			"LoadingPhase": "PreDefault"
+		},
+		{
+			"Name": "GameplayBehaviorsEditorModule",
+			"Type": "Editor",
+			"LoadingPhase": "PostEngineInit"
 		}
 	],
 	"Plugins":
diff --git a/Engine/Plugins/Experimental/GameplayBehaviors/Source/GameplayBehaviorsEditorModule/GameplayBehaviorsEditorModule.Build.cs b/Engine/Plugins/Experimental/GameplayBehaviors/Source/GameplayBehaviorsEditorModule/GameplayBehaviorsEditorModule.Build.cs
new file mode 100644
--- /dev/null
+++ b/Engine/Plugins/Experimental/GameplayBehaviors/Source/GameplayBehaviorsEditorModule/GameplayBehaviorsEditorModule.Build.cs
@@ -0,0 +1,30 @@
+// Copyright Epic Games, Inc. All Rights Reserved.
+
+namespace UnrealBuildTool.Rules
+{
+	public class GameplayBehaviorsEditorModule : ModuleRules
+	{
+		public GameplayBehaviorsEditorModule(ReadOnlyTargetRules Target) : base(Target)
+		{
+			PublicIncludePaths.AddRange(
+			new string[] {
+			}
+			);
+
+			PublicDependencyModuleNames.AddRange(
+			new string[] {
+				"Core",
+				"CoreUObject",
+				"Engine",
+				"GameplayBehaviorsModule",
+				"SlateCore"
+			}
+			);
+
+			PrivateDependencyModuleNames.AddRange(
+			new string[] {
+			}
+			);
+		}
+	}
+}
\ No newline at end of file
diff --git a/Engine/Plugins/Experimental/GameplayBehaviors/Source/GameplayBehaviorsEditorModule/Private/GameplayBehaviorsEditorModule.cpp b/Engine/Plugins/Experimental/GameplayBehaviors/Source/GameplayBehaviorsEditorModule/Private/GameplayBehaviorsEditorModule.cpp
new file mode 100644
--- /dev/null
+++ b/Engine/Plugins/Experimental/GameplayBehaviors/Source/GameplayBehaviorsEditorModule/Private/GameplayBehaviorsEditorModule.cpp
@@ -0,0 +1,28 @@
+// Copyright Epic Games, Inc. All Rights Reserved.
+
+#include "GameplayBehaviorsEditorModule.h"
+
+#include "GameplayBehaviorsEditorStyle.h"
+
+#define LOCTEXT_NAMESPACE "GameplayBehaviors"
+
+class FGameplayBehaviorsEditorModule : public IGameplayBehaviorsEditorModule
+{
+protected:
+	virtual void StartupModule() override;
+	virtual void ShutdownModule() override;
+};
+
+void FGameplayBehaviorsEditorModule::StartupModule()
+{
+	FGameplayBehaviorsEditorStyle::Get();
+}
+
+void FGameplayBehaviorsEditorModule::ShutdownModule()
+{
+	FGameplayBehaviorsEditorStyle::Shutdown();
+}
+
+IMPLEMENT_MODULE(FGameplayBehaviorsEditorModule, GameplayBehaviorsEditorModule)
+
+#undef LOCTEXT_NAMESPACE
\ No newline at end of file
diff --git a/Engine/Plugins/Experimental/GameplayBehaviors/Source/GameplayBehaviorsEditorModule/Private/GameplayBehaviorsEditorStyle.cpp b/Engine/Plugins/Experimental/GameplayBehaviors/Source/GameplayBehaviorsEditorModule/Private/GameplayBehaviorsEditorStyle.cpp
new file mode 100644
--- /dev/null
+++ b/Engine/Plugins/Experimental/GameplayBehaviors/Source/GameplayBehaviorsEditorModule/Private/GameplayBehaviorsEditorStyle.cpp
@@ -0,0 +1,32 @@
+// Copyright Epic Games, Inc. All Rights Reserved.
+
+#include "GameplayBehaviorsEditorStyle.h"
+#include "Styling/SlateStyleRegistry.h"
+
+TUniquePtr<FGameplayBehaviorsEditorStyle> FGameplayBehaviorsEditorStyle::Instance(nullptr);
+FColor FGameplayBehaviorsEditorStyle::GameplayTagTypeColor(0, 26, 154);
+
+FGameplayBehaviorsEditorStyle::FGameplayBehaviorsEditorStyle() : FSlateStyleSet("GameplayBehaviorsEditorStyle")
+{
+	Set("ClassIcon.BlackboardKeyType_GameplayTag", new FSlateRoundedBoxBrush(FLinearColor(GameplayTagTypeColor), 2.5f, FVector2D(16.f, 5.f)));
+	FSlateStyleRegistry::RegisterSlateStyle(*this);
+}
+
+FGameplayBehaviorsEditorStyle::~FGameplayBehaviorsEditorStyle()
+{
+	FSlateStyleRegistry::UnRegisterSlateStyle(*this);
+}
+
+FGameplayBehaviorsEditorStyle& FGameplayBehaviorsEditorStyle::Get()
+{
+	if (!Instance.IsValid())
+	{
+		Instance = TUniquePtr<FGameplayBehaviorsEditorStyle>(new FGameplayBehaviorsEditorStyle);
+	}
+	return *(Instance.Get());
+}
+
+void FGameplayBehaviorsEditorStyle::Shutdown()
+{
+	Instance.Release();
+}
diff --git a/Engine/Plugins/Experimental/GameplayBehaviors/Source/GameplayBehaviorsEditorModule/Private/GameplayBehaviorsEditorStyle.h b/Engine/Plugins/Experimental/GameplayBehaviors/Source/GameplayBehaviorsEditorModule/Private/GameplayBehaviorsEditorStyle.h
new file mode 100644
--- /dev/null
+++ b/Engine/Plugins/Experimental/GameplayBehaviors/Source/GameplayBehaviorsEditorModule/Private/GameplayBehaviorsEditorStyle.h
@@ -0,0 +1,22 @@
+// Copyright Epic Games, Inc. All Rights Reserved.
+
+#pragma once
+
+#include "Styling/SlateStyle.h"
+#include "Templates/UniquePtr.h"
+
+class FGameplayBehaviorsEditorStyle final : public FSlateStyleSet
+{
+public:
+	virtual ~FGameplayBehaviorsEditorStyle();
+
+	static FGameplayBehaviorsEditorStyle& Get();
+	static void Shutdown();
+
+	static FColor GameplayTagTypeColor;
+
+private:
+	FGameplayBehaviorsEditorStyle();
+
+	static TUniquePtr<FGameplayBehaviorsEditorStyle> Instance;
+};
\ No newline at end of file
diff --git a/Engine/Plugins/Experimental/GameplayBehaviors/Source/GameplayBehaviorsEditorModule/Public/GameplayBehaviorsEditorModule.h b/Engine/Plugins/Experimental/GameplayBehaviors/Source/GameplayBehaviorsEditorModule/Public/GameplayBehaviorsEditorModule.h
new file mode 100644
--- /dev/null
+++ b/Engine/Plugins/Experimental/GameplayBehaviors/Source/GameplayBehaviorsEditorModule/Public/GameplayBehaviorsEditorModule.h
@@ -0,0 +1,35 @@
+// Copyright Epic Games, Inc. All Rights Reserved.
+
+#pragma  once
+
+#include "Modules/ModuleManager.h"
+
+/** 
+* The public interface to this module
+*/
+class IGameplayBehaviorsEditorModule : public IModuleInterface
+{
+
+public:
+
+	/**
+	* Singleton-like access to this module's interface.  This is just for convenience!
+	* Beware of calling this during the shutdown phase, though.  Your module might have been unloaded already.
+	*
+	* @return Returns singleton instance, loading the module on demand if needed
+	*/
+	static inline IGameplayBehaviorsEditorModule& Get()
+	{
+		return FModuleManager::LoadModuleChecked<IGameplayBehaviorsEditorModule>("GameplayBehaviorsEditorModule");
+	}
+
+	/**
+	 * Checks to see if this module is loaded and ready.  It is only valid to call Get() if IsAvailable() returns true.
+	 *
+	 * @return True if the module is loaded and ready to use
+	 */
+	static inline bool IsAvailable()
+	{
+		return FModuleManager::Get().IsModuleLoaded("GameplayBehaviorsEditorModule");
+	}
+};
\ No newline at end of file
diff --git a/Engine/Plugins/Experimental/GameplayBehaviors/Source/GameplayBehaviorsModule/Private/BlackboardKeyType_GameplayTag.cpp b/Engine/Plugins/Experimental/GameplayBehaviors/Source/GameplayBehaviorsModule/Private/BlackboardKeyType_GameplayTag.cpp
new file mode 100644
--- /dev/null
+++ b/Engine/Plugins/Experimental/GameplayBehaviors/Source/GameplayBehaviorsModule/Private/BlackboardKeyType_GameplayTag.cpp
@@ -0,0 +1,56 @@
+// Copyright Epic Games, Inc. All Rights Reserved.
+
+#include "BlackboardKeyType_GameplayTag.h"
+
+const UBlackboardKeyType_GameplayTag::FDataType UBlackboardKeyType_GameplayTag::InvalidValue = FGameplayTagContainer::EmptyContainer;
+
+UBlackboardKeyType_GameplayTag::UBlackboardKeyType_GameplayTag(const FObjectInitializer& ObjectInitializer) : Super(ObjectInitializer)
+{
+	ValueSize = sizeof(FGameplayTagContainer);
+	SupportedOp = EBlackboardKeyOperation::Text;
+}
+
+FGameplayTagContainer UBlackboardKeyType_GameplayTag::GetValue(const UBlackboardKeyType_GameplayTag* KeyOb, const uint8* RawData)
+{
+	return GetValueFromMemory<FGameplayTagContainer>(RawData);
+}
+
+bool UBlackboardKeyType_GameplayTag::SetValue(UBlackboardKeyType_GameplayTag* KeyOb, uint8* RawData, const FGameplayTagContainer& Value)
+{
+	return SetValueInMemory<FGameplayTagContainer>(RawData, Value);
+}
+
+EBlackboardCompare::Type UBlackboardKeyType_GameplayTag::CompareValues(const UBlackboardComponent& OwnerComp, const uint8* MemoryBlock,
+	const UBlackboardKeyType* OtherKeyOb, const uint8* OtherMemoryBlock) const
+{
+	const FGameplayTagContainer MyValue = GetValue(this, MemoryBlock);
+	const FGameplayTagContainer OtherValue = GetValue((UBlackboardKeyType_GameplayTag*)OtherKeyOb, OtherMemoryBlock);
+
+	return (MyValue.HasAll(OtherValue) ? EBlackboardCompare::Equal : EBlackboardCompare::NotEqual);
+}
+
+FString UBlackboardKeyType_GameplayTag::DescribeValue(const UBlackboardComponent& OwnerComp, const uint8* RawData) const
+{
+	const FGameplayTagContainer CurrentValue = GetValue(this, RawData);
+	if (CurrentValue.IsEmpty())
+	{
+		static const FString EmptyContainer = TEXT("Empty");
+		return EmptyContainer;
+	}
+	return CurrentValue.ToStringSimple();
+}
+
+bool UBlackboardKeyType_GameplayTag::TestTextOperation(const UBlackboardComponent& OwnerComp, const uint8* MemoryBlock, ETextKeyOperation::Type Op, const FString& OtherString) const
+{
+	const FString Value = GetValue(this, MemoryBlock).ToStringSimple();
+	switch (Op)
+	{
+	case ETextKeyOperation::Equal:			return (Value == OtherString);
+	case ETextKeyOperation::NotEqual:		return (Value != OtherString);
+	case ETextKeyOperation::Contain:		return (Value.Contains(OtherString) == true);
+	case ETextKeyOperation::NotContain:		return (Value.Contains(OtherString) == false);
+	default: break;
+	}
+
+	return false;
+}
\ No newline at end of file
diff --git a/Engine/Plugins/Experimental/GameplayBehaviors/Source/GameplayBehaviorsModule/Private/GameplayBehaviorsBlueprintFunctionLibrary.cpp b/Engine/Plugins/Experimental/GameplayBehaviors/Source/GameplayBehaviorsModule/Private/GameplayBehaviorsBlueprintFunctionLibrary.cpp
--- a/Engine/Plugins/Experimental/GameplayBehaviors/Source/GameplayBehaviorsModule/Private/GameplayBehaviorsBlueprintFunctionLibrary.cpp
+++ b/Engine/Plugins/Experimental/GameplayBehaviors/Source/GameplayBehaviorsModule/Private/GameplayBehaviorsBlueprintFunctionLibrary.cpp
@@ -1,9 +1,12 @@
 // Copyright Epic Games, Inc. All Rights Reserved.
 
 #include "GameplayBehaviorsBlueprintFunctionLibrary.h"
+#include "BehaviorTree/BlackboardComponent.h"
+#include "BehaviorTree/BTNode.h"
+#include "BehaviorTree/BTFunctionLibrary.h"
+#include "BlackboardKeyType_GameplayTag.h"
 #include "GameFramework/Actor.h"
 
-
 bool UGameplayBehaviorsBlueprintFunctionLibrary::StopGameplayBehavior(TSubclassOf<UGameplayBehavior> GameplayBehaviorClass, AActor* Avatar)
 {
 	if (Avatar == nullptr || !GameplayBehaviorClass)
@@ -12,4 +15,37 @@ bool UGameplayBehaviorsBlueprintFunctionLibrary::StopGameplayBehavior(TSubclassO
 	}
 
 	return false;
-}
\ No newline at end of file
+}
+
+FGameplayTagContainer UGameplayBehaviorsBlueprintFunctionLibrary::GetBlackboardValueAsGameplayTag(UBTNode* NodeOwner, const FBlackboardKeySelector& Key)
+{
+	UBlackboardComponent* BlackboardComp = UBTFunctionLibrary::GetOwnersBlackboard(NodeOwner);
+	return GetBlackboardValueAsGameplayTagFromBlackboardComp(BlackboardComp, Key.SelectedKeyName);
+}
+
+void UGameplayBehaviorsBlueprintFunctionLibrary::SetBlackboardValueAsGameplayTag(UBTNode* NodeOwner, const FBlackboardKeySelector& Key, FGameplayTagContainer Value)
+{
+	UBlackboardComponent* BlackboardComp = UBTFunctionLibrary::GetOwnersBlackboard(NodeOwner);
+	SetValueAsGameplayTagForBlackboardComp(BlackboardComp, Key.SelectedKeyName, Value);
+}
+
+void UGameplayBehaviorsBlueprintFunctionLibrary::AddGameplayTagFilterToBlackboardKeySelector(FBlackboardKeySelector& InSelector, UObject* Owner, FName PropertyName)
+{
+	const FString FilterName = PropertyName.ToString() + TEXT("_GameplayTag");
+	InSelector.AllowedTypes.Add(NewObject<UBlackboardKeyType_GameplayTag>(Owner, *FilterName));
+}
+
+FGameplayTagContainer UGameplayBehaviorsBlueprintFunctionLibrary::GetBlackboardValueAsGameplayTagFromBlackboardComp(UBlackboardComponent* BlackboardComp, const FName& KeyName)
+{
+	return BlackboardComp ? BlackboardComp->GetValue<UBlackboardKeyType_GameplayTag>(KeyName) : FGameplayTagContainer::EmptyContainer;
+}
+
+
+void UGameplayBehaviorsBlueprintFunctionLibrary::SetValueAsGameplayTagForBlackboardComp(UBlackboardComponent* BlackboardComp, const FName& KeyName, FGameplayTagContainer GameplayTagValue)
+{
+	if (BlackboardComp)
+	{
+		const FBlackboard::FKey KeyID = BlackboardComp->GetKeyID(KeyName);
+		BlackboardComp->SetValue<UBlackboardKeyType_GameplayTag>(KeyID, GameplayTagValue);
+	}
+}
diff --git a/Engine/Plugins/Experimental/GameplayBehaviors/Source/GameplayBehaviorsModule/Public/BlackboardKeyType_GameplayTag.h b/Engine/Plugins/Experimental/GameplayBehaviors/Source/GameplayBehaviorsModule/Public/BlackboardKeyType_GameplayTag.h
new file mode 100644
--- /dev/null
+++ b/Engine/Plugins/Experimental/GameplayBehaviors/Source/GameplayBehaviorsModule/Public/BlackboardKeyType_GameplayTag.h
@@ -0,0 +1,32 @@
+// Copyright Epic Games, Inc. All Rights Reserved.
+#pragma once
+
+#include "CoreMinimal.h"
+#include "UObject/ObjectMacros.h"
+#include "BehaviorTree/Blackboard/BlackboardKeyType.h"
+#include "GameplayTagContainer.h"
+#include "BlackboardKeyType_GameplayTag.generated.h"
+
+class UBlackboardComponent;
+
+UCLASS(EditInlineNew, meta = (DisplayName = "Gameplay Tag"))
+class GAMEPLAYBEHAVIORSMODULE_API UBlackboardKeyType_GameplayTag : public UBlackboardKeyType
+{
+	GENERATED_BODY()
+public:
+
+	UBlackboardKeyType_GameplayTag(const FObjectInitializer& ObjectInitializer);
+
+	typedef FGameplayTagContainer FDataType;
+	static const FDataType InvalidValue;
+
+	static FGameplayTagContainer GetValue(const UBlackboardKeyType_GameplayTag* KeyOb, const uint8* RawData);
+	static bool SetValue(UBlackboardKeyType_GameplayTag* KeyOb, uint8* RawData, const FGameplayTagContainer& Value);
+
+	virtual EBlackboardCompare::Type CompareValues(const UBlackboardComponent& OwnerComp, const uint8* MemoryBlock,
+		const UBlackboardKeyType* OtherKeyOb, const uint8* OtherMemoryBlock) const override;
+
+protected:
+	virtual FString DescribeValue(const UBlackboardComponent& OwnerComp, const uint8* RawData) const override;
+	virtual bool TestTextOperation(const UBlackboardComponent& OwnerComp, const uint8* MemoryBlock, ETextKeyOperation::Type Op, const FString& OtherString) const override;
+};
diff --git a/Engine/Plugins/Experimental/GameplayBehaviors/Source/GameplayBehaviorsModule/Public/GameplayBehaviorsBlueprintFunctionLibrary.h b/Engine/Plugins/Experimental/GameplayBehaviors/Source/GameplayBehaviorsModule/Public/GameplayBehaviorsBlueprintFunctionLibrary.h
--- a/Engine/Plugins/Experimental/GameplayBehaviors/Source/GameplayBehaviorsModule/Public/GameplayBehaviorsBlueprintFunctionLibrary.h
+++ b/Engine/Plugins/Experimental/GameplayBehaviors/Source/GameplayBehaviorsModule/Public/GameplayBehaviorsBlueprintFunctionLibrary.h
@@ -3,13 +3,16 @@
 #pragma once
 
 #include "CoreMinimal.h"
-#include "UObject/ObjectMacros.h"
-#include "Kismet/BlueprintFunctionLibrary.h"
+#include "BehaviorTree/BehaviorTreeTypes.h"
 #include "GameplayBehavior.h"
+#include "GameplayTagContainer.h"
+#include "Kismet/BlueprintFunctionLibrary.h"
+#include "UObject/ObjectMacros.h"
 #include "GameplayBehaviorsBlueprintFunctionLibrary.generated.h"
 
 
 class AActor;
+class UBTNode;
 
 UCLASS(meta = (ScriptName = "GameplayBehaviorsLibrary"))
 class GAMEPLAYBEHAVIORSMODULE_API UGameplayBehaviorsBlueprintFunctionLibrary : public UBlueprintFunctionLibrary
@@ -19,4 +22,20 @@ public:
 	/** Will force-stop GameplayBehavior on given Avatar assuming the current 
 	 *	behavior is of GameplayBehaviorClass class*/
 	static bool StopGameplayBehavior(TSubclassOf<UGameplayBehavior> GameplayBehaviorClass,  AActor* Avatar);
+
+	UFUNCTION(BlueprintPure, Category = "AI|BehaviorTree", Meta = (HidePin = "NodeOwner", DefaultToSelf = "NodeOwner"))
+	static FGameplayTagContainer GetBlackboardValueAsGameplayTag(UBTNode* NodeOwner, const FBlackboardKeySelector& Key);
+
+	UFUNCTION(BlueprintCallable, Category = "AI|BehaviorTree", Meta = (HidePin = "NodeOwner", DefaultToSelf = "NodeOwner"))
+	static void SetBlackboardValueAsGameplayTag(UBTNode* NodeOwner, const FBlackboardKeySelector& Key, FGameplayTagContainer Value);
+
+	UFUNCTION(BlueprintCallable, Category = "AI|BehaviorTree")
+	static void AddGameplayTagFilterToBlackboardKeySelector(FBlackboardKeySelector& InSelector, UObject* Owner, FName PropertyName);
+
+	UFUNCTION(BlueprintCallable, Category = "AI|BehaviorTree")
+	static FGameplayTagContainer GetBlackboardValueAsGameplayTagFromBlackboardComp(UBlackboardComponent* BlackboardComp, const FName& KeyName);
+
+	UFUNCTION(BlueprintCallable, Category = "AI|BehaviorTree")
+	static void SetValueAsGameplayTagForBlackboardComp(UBlackboardComponent* BlackboardComp, const FName& KeyName, FGameplayTagContainer GameplayTagValue);
+
 };
diff --git a/Engine/Source/Editor/BehaviorTreeEditor/Private/DetailCustomizations/BlackboardSelectorDetails.cpp b/Engine/Source/Editor/BehaviorTreeEditor/Private/DetailCustomizations/BlackboardSelectorDetails.cpp
--- a/Engine/Source/Editor/BehaviorTreeEditor/Private/DetailCustomizations/BlackboardSelectorDetails.cpp
+++ b/Engine/Source/Editor/BehaviorTreeEditor/Private/DetailCustomizations/BlackboardSelectorDetails.cpp
@@ -108,6 +108,11 @@ void FBlackboardSelectorDetails::CacheBlackboardData()
 	
 	KeyValues.Reset();
 
+	if (bNoneIsAllowedValue)
+	{
+		KeyValues.AddUnique(TEXT("None"));
+	}
+
 	TArray<UBlackboardKeyType*> FilterObjects;
 	
 	uint32 NumElements = 0;
