Changed UGameplayBehaviorsBlueprintFunctionLibrary's ScriptName to GameplayBehaviorLibrary to fix name collision warning

diff --git a/Engine/Plugins/Experimental/GameplayBehaviors/Source/GameplayBehaviorsModule/Public/GameplayBehaviorsBlueprintFunctionLibrary.h b/Engine/Plugins/Experimental/GameplayBehaviors/Source/GameplayBehaviorsModule/Public/GameplayBehaviorsBlueprintFunctionLibrary.h
--- a/Engine/Plugins/Experimental/GameplayBehaviors/Source/GameplayBehaviorsModule/Public/GameplayBehaviorsBlueprintFunctionLibrary.h
+++ b/Engine/Plugins/Experimental/GameplayBehaviors/Source/GameplayBehaviorsModule/Public/GameplayBehaviorsBlueprintFunctionLibrary.h
@@ -14,7 +14,7 @@
 class AActor;
 class UBTNode;
 
-UCLASS(meta = (ScriptName = "GameplayBehaviorsLibrary"))
+UCLASS(meta = (ScriptName = "GameplayBehaviorLibrary"))
 class GAMEPLAYBEHAVIORSMODULE_API UGameplayBehaviorsBlueprintFunctionLibrary : public UBlueprintFunctionLibrary
 {
 	GENERATED_BODY()
