GameplayTagQuery: Move struct that doesn't need to be public from header to cpp file

#tests PIE
#rnx

diff --git a/Engine/Plugins/Experimental/GameplayBehaviors/Source/GameplayBehaviorsModule/Private/AI/BTDecorator_GameplayTagQuery.cpp b/Engine/Plugins/Experimental/GameplayBehaviors/Source/GameplayBehaviorsModule/Private/AI/BTDecorator_GameplayTagQuery.cpp
--- a/Engine/Plugins/Experimental/GameplayBehaviors/Source/GameplayBehaviorsModule/Private/AI/BTDecorator_GameplayTagQuery.cpp
+++ b/Engine/Plugins/Experimental/GameplayBehaviors/Source/GameplayBehaviorsModule/Private/AI/BTDecorator_GameplayTagQuery.cpp
@@ -12,6 +12,14 @@
 
 #include UE_INLINE_GENERATED_CPP_BY_NAME(BTDecorator_GameplayTagQuery)
 
+struct FBTDecorator_GameplayTagQueryMemory
+{
+	TWeakObjectPtr<UAbilitySystemComponent> CachedAbilitySystemComponent;
+
+	/** Array of handles for our gameplay tag query delegates */
+	TArray<TTuple<FGameplayTag, FDelegateHandle>> GameplayTagEventHandles;
+};
+
 UBTDecorator_GameplayTagQuery::UBTDecorator_GameplayTagQuery(const FObjectInitializer& ObjectInitializer) : Super(ObjectInitializer)
 {
 	NodeName = "Gameplay Tag Query";
diff --git a/Engine/Plugins/Experimental/GameplayBehaviors/Source/GameplayBehaviorsModule/Public/AI/BTDecorator_GameplayTagQuery.h b/Engine/Plugins/Experimental/GameplayBehaviors/Source/GameplayBehaviorsModule/Public/AI/BTDecorator_GameplayTagQuery.h
--- a/Engine/Plugins/Experimental/GameplayBehaviors/Source/GameplayBehaviorsModule/Public/AI/BTDecorator_GameplayTagQuery.h
+++ b/Engine/Plugins/Experimental/GameplayBehaviors/Source/GameplayBehaviorsModule/Public/AI/BTDecorator_GameplayTagQuery.h
@@ -17,14 +17,6 @@ class UBehaviorTreeComponent;
 class UBlackboardComponent;
 class UAbilitySystemComponent;
 
-struct FBTDecorator_GameplayTagQueryMemory
-{
-	TWeakObjectPtr<UAbilitySystemComponent> CachedAbilitySystemComponent;
-	
-	/** Array of handles for our gameplay tag query delegates */
-	TArray<TTuple<FGameplayTag, FDelegateHandle>> GameplayTagEventHandles;
-};
-
 /**
  * GameplayTagQuery decorator node.
  * A decorator node that bases its condition on matching a gameplay tag query.
