- Adding copy, initialize memory, clear, and is empty overridden functions to hopefully fix memory corruption crash

diff --git a/Engine/Plugins/Experimental/GameplayBehaviors/Source/GameplayBehaviorsModule/Private/BlackboardKeyType_GameplayTag.cpp b/Engine/Plugins/Experimental/GameplayBehaviors/Source/GameplayBehaviorsModule/Private/BlackboardKeyType_GameplayTag.cpp
--- a/Engine/Plugins/Experimental/GameplayBehaviors/Source/GameplayBehaviorsModule/Private/BlackboardKeyType_GameplayTag.cpp
+++ b/Engine/Plugins/Experimental/GameplayBehaviors/Source/GameplayBehaviorsModule/Private/BlackboardKeyType_GameplayTag.cpp
@@ -56,3 +56,29 @@ bool UBlackboardKeyType_GameplayTag::TestTextOperation(const UBlackboardComponen
 
 	return false;
 }
+
+void UBlackboardKeyType_GameplayTag::CopyValues(UBlackboardComponent& OwnerComp, uint8* MemoryBlock, const UBlackboardKeyType* SourceKeyOb, const uint8* SourceBlock)
+{
+	const UBlackboardKeyType_GameplayTag* GameplayTagSourceKeyOb = Cast<const UBlackboardKeyType_GameplayTag>(SourceKeyOb);
+	if (GameplayTagSourceKeyOb != nullptr)
+	{
+		SetValue(this, MemoryBlock, GetValue(GameplayTagSourceKeyOb, SourceBlock));
+	}
+}
+
+void UBlackboardKeyType_GameplayTag::InitializeMemory(UBlackboardComponent& OwnerComp, uint8* MemoryBlock)
+{
+	SetValue(this, MemoryBlock, InvalidValue);
+}
+
+void UBlackboardKeyType_GameplayTag::Clear(UBlackboardComponent& OwnerComp, uint8* MemoryBlock)
+{
+	SetValue(this, MemoryBlock, InvalidValue);
+}
+
+bool UBlackboardKeyType_GameplayTag::IsEmpty(const UBlackboardComponent& OwnerComp, const uint8* MemoryBlock) const
+{
+	const FDataType TagContainer = GetValue(this, MemoryBlock);
+
+	return TagContainer.IsEmpty();
+}
diff --git a/Engine/Plugins/Experimental/GameplayBehaviors/Source/GameplayBehaviorsModule/Public/BlackboardKeyType_GameplayTag.h b/Engine/Plugins/Experimental/GameplayBehaviors/Source/GameplayBehaviorsModule/Public/BlackboardKeyType_GameplayTag.h
--- a/Engine/Plugins/Experimental/GameplayBehaviors/Source/GameplayBehaviorsModule/Public/BlackboardKeyType_GameplayTag.h
+++ b/Engine/Plugins/Experimental/GameplayBehaviors/Source/GameplayBehaviorsModule/Public/BlackboardKeyType_GameplayTag.h
@@ -29,4 +29,8 @@ public:
 protected:
 	virtual FString DescribeValue(const UBlackboardComponent& OwnerComp, const uint8* RawData) const override;
 	virtual bool TestTextOperation(const UBlackboardComponent& OwnerComp, const uint8* MemoryBlock, ETextKeyOperation::Type Op, const FString& OtherString) const override;
+	virtual void CopyValues(UBlackboardComponent& OwnerComp, uint8* MemoryBlock, const UBlackboardKeyType* SourceKeyOb, const uint8* SourceBlock) override;
+	virtual void InitializeMemory(UBlackboardComponent& OwnerComp, uint8* MemoryBlock) override;
+	virtual void Clear(UBlackboardComponent& OwnerComp, uint8* MemoryBlock) override;
+	virtual bool IsEmpty(const UBlackboardComponent& OwnerComp, const uint8* MemoryBlock) const override;
 };
