Support calling world-dependent nodes (i.e. Async Load Asset) in Smart Object Behaviors

diff --git a/Engine/Plugins/Experimental/GameplayBehaviors/Source/GameplayBehaviorsModule/Private/GameplayBehavior.cpp b/Engine/Plugins/Experimental/GameplayBehaviors/Source/GameplayBehaviorsModule/Private/GameplayBehavior.cpp
--- a/Engine/Plugins/Experimental/GameplayBehaviors/Source/GameplayBehaviorsModule/Private/GameplayBehavior.cpp
+++ b/Engine/Plugins/Experimental/GameplayBehaviors/Source/GameplayBehaviorsModule/Private/GameplayBehavior.cpp
@@ -58,6 +58,29 @@ void UGameplayBehavior::BeginDestroy()
 	Super::BeginDestroy();
 }
 
+UWorld* UGameplayBehavior::GetWorld() const
+{
+	// Currently Gameplay Behaviors are created with the World as their outer, so for instances this should always work
+	if (const UObject* Outer = GetOuter())
+	{
+		if (UWorld* World = Outer->GetWorld())
+		{
+			return World;
+		}
+	}
+
+	// The CDO has no World outer, so let's take the World from the affected actor,
+	// which will be available in all BP Trigger events
+	if (TransientAvatar)
+	{
+		return TransientAvatar->GetWorld();
+	}
+
+	// Do not call super, as that'll indicate to the BP editor that GetWorld() is not implemented,
+	// and will result in an error that calling world-dependent functions is unsafe
+	return nullptr;
+}
+
 bool UGameplayBehavior::Trigger(AActor& Avatar, const UGameplayBehaviorConfig* Config, AActor* SmartObjectOwner/* = nullptr*/)
 {
 	bTransientIsTriggering = true;
diff --git a/Engine/Plugins/Experimental/GameplayBehaviors/Source/GameplayBehaviorsModule/Public/GameplayBehavior.h b/Engine/Plugins/Experimental/GameplayBehaviors/Source/GameplayBehaviorsModule/Public/GameplayBehavior.h
--- a/Engine/Plugins/Experimental/GameplayBehaviors/Source/GameplayBehaviorsModule/Public/GameplayBehavior.h
+++ b/Engine/Plugins/Experimental/GameplayBehaviors/Source/GameplayBehaviorsModule/Public/GameplayBehavior.h
@@ -39,6 +39,7 @@ public:
 
 	virtual void PostInitProperties() override;
 	virtual void BeginDestroy() override;
+	virtual UWorld* GetWorld() const override;
 
 	/**
 	 * Default implementation will trigger the appropriate Blueprint event based on 
