Add ability to override default blackboard key value for most type.
Prevent child blackboards from being able to modify the key type of inherited key.

diff --git a/Engine/Plugins/Experimental/GameplayBehaviors/Source/GameplayBehaviorsModule/Private/BlackboardKeyType_GameplayTag.cpp b/Engine/Plugins/Experimental/GameplayBehaviors/Source/GameplayBehaviorsModule/Private/BlackboardKeyType_GameplayTag.cpp
--- a/Engine/Plugins/Experimental/GameplayBehaviors/Source/GameplayBehaviorsModule/Private/BlackboardKeyType_GameplayTag.cpp
+++ b/Engine/Plugins/Experimental/GameplayBehaviors/Source/GameplayBehaviorsModule/Private/BlackboardKeyType_GameplayTag.cpp
@@ -69,7 +69,7 @@ void UBlackboardKeyType_GameplayTag::CopyValues(UBlackboardComponent& OwnerComp,
 
 void UBlackboardKeyType_GameplayTag::InitializeMemory(UBlackboardComponent& OwnerComp, uint8* MemoryBlock)
 {
-	SetValue(this, MemoryBlock, InvalidValue);
+	SetValue(this, MemoryBlock, DefaultValue);
 }
 
 void UBlackboardKeyType_GameplayTag::Clear(UBlackboardComponent& OwnerComp, uint8* MemoryBlock)
diff --git a/Engine/Plugins/Experimental/GameplayBehaviors/Source/GameplayBehaviorsModule/Public/BlackboardKeyType_GameplayTag.h b/Engine/Plugins/Experimental/GameplayBehaviors/Source/GameplayBehaviorsModule/Public/BlackboardKeyType_GameplayTag.h
--- a/Engine/Plugins/Experimental/GameplayBehaviors/Source/GameplayBehaviorsModule/Public/BlackboardKeyType_GameplayTag.h
+++ b/Engine/Plugins/Experimental/GameplayBehaviors/Source/GameplayBehaviorsModule/Public/BlackboardKeyType_GameplayTag.h
@@ -3,10 +3,9 @@
 
 #include "BehaviorTree/Blackboard/BlackboardKeyEnums.h"
 #include "BehaviorTree/Blackboard/BlackboardKeyType.h"
+#include "GameplayTagContainer.h"
 #include "BlackboardKeyType_GameplayTag.generated.h"
 
-struct FGameplayTagContainer;
-
 class UBlackboardComponent;
 
 UCLASS(EditInlineNew, meta = (DisplayName = "Gameplay Tag"))
@@ -26,6 +25,9 @@ public:
 	virtual EBlackboardCompare::Type CompareValues(const UBlackboardComponent& OwnerComp, const uint8* MemoryBlock,
 		const UBlackboardKeyType* OtherKeyOb, const uint8* OtherMemoryBlock) const override;
 
+	UPROPERTY(EditDefaultsOnly, Category = Blackboard)
+	FGameplayTagContainer DefaultValue = InvalidValue;
+
 protected:
 	virtual FString DescribeValue(const UBlackboardComponent& OwnerComp, const uint8* RawData) const override;
 	virtual bool TestTextOperation(const UBlackboardComponent& OwnerComp, const uint8* MemoryBlock, ETextKeyOperation::Type Op, const FString& OtherString) const override;
diff --git a/Engine/Source/Editor/BehaviorTreeEditor/Private/DetailCustomizations/BlackboardDataDetails.cpp b/Engine/Source/Editor/BehaviorTreeEditor/Private/DetailCustomizations/BlackboardDataDetails.cpp
--- a/Engine/Source/Editor/BehaviorTreeEditor/Private/DetailCustomizations/BlackboardDataDetails.cpp
+++ b/Engine/Source/Editor/BehaviorTreeEditor/Private/DetailCustomizations/BlackboardDataDetails.cpp
@@ -130,7 +130,8 @@ void FBlackboardDataDetails::CustomizeDetails( IDetailLayoutBuilder& DetailLayou
 			DetailCategoryBuilder.AddProperty(EntryDescriptionHandle);
 
 			TSharedPtr<IPropertyHandle> KeyTypeProperty = KeyHandle->GetChildHandle(GET_MEMBER_NAME_CHECKED(FBlackboardEntry, KeyType));
-			DetailCategoryBuilder.AddProperty(KeyTypeProperty);
+			DetailCategoryBuilder.AddProperty(KeyTypeProperty)
+			.EditCondition(!bIsInherited, nullptr); /** nullptr because the bool condition will never change so no need for a listener. */
 
 			TSharedPtr<IPropertyHandle> bInstanceSyncedProperty = KeyHandle->GetChildHandle(GET_MEMBER_NAME_CHECKED(FBlackboardEntry, bInstanceSynced));
 			DetailCategoryBuilder.AddProperty(bInstanceSyncedProperty);
diff --git a/Engine/Source/Runtime/AIModule/Classes/BehaviorTree/Blackboard/BlackboardKeyType_Bool.h b/Engine/Source/Runtime/AIModule/Classes/BehaviorTree/Blackboard/BlackboardKeyType_Bool.h
--- a/Engine/Source/Runtime/AIModule/Classes/BehaviorTree/Blackboard/BlackboardKeyType_Bool.h
+++ b/Engine/Source/Runtime/AIModule/Classes/BehaviorTree/Blackboard/BlackboardKeyType_Bool.h
@@ -22,7 +22,11 @@ class UBlackboardKeyType_Bool : public UBlackboardKeyType
 	AIMODULE_API virtual EBlackboardCompare::Type CompareValues(const UBlackboardComponent& OwnerComp, const uint8* MemoryBlock,
 		const UBlackboardKeyType* OtherKeyOb, const uint8* OtherMemoryBlock) const override;
 
+	UPROPERTY(Category = Blackboard, EditDefaultsOnly)
+	bool bDefaultValue = InvalidValue;
+
 protected:
+	AIMODULE_API virtual void InitializeMemory(UBlackboardComponent& OwnerComp, uint8* MemoryBlock) override;
 	AIMODULE_API virtual FString DescribeValue(const UBlackboardComponent& OwnerComp, const uint8* RawData) const override;
 	AIMODULE_API virtual bool TestBasicOperation(const UBlackboardComponent& OwnerComp, const uint8* MemoryBlock, EBasicKeyOperation::Type Op) const override;
 };
diff --git a/Engine/Source/Runtime/AIModule/Classes/BehaviorTree/Blackboard/BlackboardKeyType_Float.h b/Engine/Source/Runtime/AIModule/Classes/BehaviorTree/Blackboard/BlackboardKeyType_Float.h
--- a/Engine/Source/Runtime/AIModule/Classes/BehaviorTree/Blackboard/BlackboardKeyType_Float.h
+++ b/Engine/Source/Runtime/AIModule/Classes/BehaviorTree/Blackboard/BlackboardKeyType_Float.h
@@ -24,7 +24,11 @@ class UBlackboardKeyType_Float : public UBlackboardKeyType
 
 	AIMODULE_API virtual FString DescribeArithmeticParam(int32 IntValue, float FloatValue) const override;
 
+	UPROPERTY(EditDefaultsOnly, Category=Blackboard)
+	float DefaultValue = InvalidValue;
+
 protected:
+	AIMODULE_API virtual void InitializeMemory(UBlackboardComponent& OwnerComp, uint8* MemoryBlock) override;
 	AIMODULE_API virtual FString DescribeValue(const UBlackboardComponent& OwnerComp, const uint8* RawData) const override;
 	AIMODULE_API virtual bool TestArithmeticOperation(const UBlackboardComponent& OwnerComp, const uint8* MemoryBlock, EArithmeticKeyOperation::Type Op, int32 OtherIntValue, float OtherFloatValue) const override;
 };
diff --git a/Engine/Source/Runtime/AIModule/Classes/BehaviorTree/Blackboard/BlackboardKeyType_Int.h b/Engine/Source/Runtime/AIModule/Classes/BehaviorTree/Blackboard/BlackboardKeyType_Int.h
--- a/Engine/Source/Runtime/AIModule/Classes/BehaviorTree/Blackboard/BlackboardKeyType_Int.h
+++ b/Engine/Source/Runtime/AIModule/Classes/BehaviorTree/Blackboard/BlackboardKeyType_Int.h
@@ -23,7 +23,11 @@ class UBlackboardKeyType_Int : public UBlackboardKeyType
 	AIMODULE_API virtual EBlackboardCompare::Type CompareValues(const UBlackboardComponent& OwnerComp, const uint8* MemoryBlock,
 		const UBlackboardKeyType* OtherKeyOb, const uint8* OtherMemoryBlock) const override;
 
+	UPROPERTY(EditDefaultsOnly, Category = Blackboard)
+	int32 DefaultValue = InvalidValue;
+
 protected:
+	AIMODULE_API virtual void InitializeMemory(UBlackboardComponent& OwnerComp, uint8* MemoryBlock) override;
 	AIMODULE_API virtual FString DescribeValue(const UBlackboardComponent& OwnerComp, const uint8* RawData) const override;
 	AIMODULE_API virtual bool TestArithmeticOperation(const UBlackboardComponent& OwnerComp, const uint8* MemoryBlock, EArithmeticKeyOperation::Type Op, int32 OtherIntValue, float OtherFloatValue) const override;
 };
diff --git a/Engine/Source/Runtime/AIModule/Classes/BehaviorTree/Blackboard/BlackboardKeyType_Name.h b/Engine/Source/Runtime/AIModule/Classes/BehaviorTree/Blackboard/BlackboardKeyType_Name.h
--- a/Engine/Source/Runtime/AIModule/Classes/BehaviorTree/Blackboard/BlackboardKeyType_Name.h
+++ b/Engine/Source/Runtime/AIModule/Classes/BehaviorTree/Blackboard/BlackboardKeyType_Name.h
@@ -22,7 +22,11 @@ class UBlackboardKeyType_Name : public UBlackboardKeyType
 	AIMODULE_API virtual EBlackboardCompare::Type CompareValues(const UBlackboardComponent& OwnerComp, const uint8* MemoryBlock,
 		const UBlackboardKeyType* OtherKeyOb, const uint8* OtherMemoryBlock) const override;
 
+	UPROPERTY(EditDefaultsOnly, Category = Blackboard)
+	FName DefaultValue = InvalidValue;
+
 protected:
+	AIMODULE_API virtual void InitializeMemory(UBlackboardComponent& OwnerComp, uint8* MemoryBlock) override;
 	AIMODULE_API virtual FString DescribeValue(const UBlackboardComponent& OwnerComp, const uint8* RawData) const override;
 	AIMODULE_API virtual bool TestTextOperation(const UBlackboardComponent& OwnerComp, const uint8* MemoryBlock, ETextKeyOperation::Type Op, const FString& OtherString) const override;
 };
diff --git a/Engine/Source/Runtime/AIModule/Classes/BehaviorTree/Blackboard/BlackboardKeyType_Rotator.h b/Engine/Source/Runtime/AIModule/Classes/BehaviorTree/Blackboard/BlackboardKeyType_Rotator.h
--- a/Engine/Source/Runtime/AIModule/Classes/BehaviorTree/Blackboard/BlackboardKeyType_Rotator.h
+++ b/Engine/Source/Runtime/AIModule/Classes/BehaviorTree/Blackboard/BlackboardKeyType_Rotator.h
@@ -22,6 +22,13 @@ class UBlackboardKeyType_Rotator : public UBlackboardKeyType
 	AIMODULE_API virtual EBlackboardCompare::Type CompareValues(const UBlackboardComponent& OwnerComp, const uint8* MemoryBlock,
 		const UBlackboardKeyType* OtherKeyOb, const uint8* OtherMemoryBlock) const override;
 
+	// If not set, will default to  FAISystem::InvalidRotation
+	UPROPERTY(EditDefaultsOnly, Category = Blackboard, meta = (EditCondition = "bUseDefaultValue"))
+	FRotator DefaultValue = FRotator::ZeroRotator;
+
+	UPROPERTY()
+	bool bUseDefaultValue = false;
+
 protected:
 	AIMODULE_API virtual void InitializeMemory(UBlackboardComponent& OwnerComp, uint8* RawData) override;
 	AIMODULE_API virtual FString DescribeValue(const UBlackboardComponent& OwnerComp, const uint8* RawData) const override;
diff --git a/Engine/Source/Runtime/AIModule/Classes/BehaviorTree/Blackboard/BlackboardKeyType_String.h b/Engine/Source/Runtime/AIModule/Classes/BehaviorTree/Blackboard/BlackboardKeyType_String.h
--- a/Engine/Source/Runtime/AIModule/Classes/BehaviorTree/Blackboard/BlackboardKeyType_String.h
+++ b/Engine/Source/Runtime/AIModule/Classes/BehaviorTree/Blackboard/BlackboardKeyType_String.h
@@ -25,7 +25,11 @@ class UBlackboardKeyType_String : public UBlackboardKeyType
 	UPROPERTY()
 	FString StringValue;
 
+	UPROPERTY(EditDefaultsOnly, Category = Blackboard)
+	FString DefaultValue = InvalidValue;
+
 protected:
+	AIMODULE_API virtual void InitializeMemory(UBlackboardComponent& OwnerComp, uint8* MemoryBlock) override;
 	AIMODULE_API virtual FString DescribeValue(const UBlackboardComponent& OwnerComp, const uint8* RawData) const override;
 	AIMODULE_API virtual bool TestTextOperation(const UBlackboardComponent& OwnerComp, const uint8* MemoryBlock, ETextKeyOperation::Type Op, const FString& OtherString) const override;
 	AIMODULE_API virtual void Clear(UBlackboardComponent& OwnerComp, uint8* MemoryBlock) override;
diff --git a/Engine/Source/Runtime/AIModule/Classes/BehaviorTree/Blackboard/BlackboardKeyType_Vector.h b/Engine/Source/Runtime/AIModule/Classes/BehaviorTree/Blackboard/BlackboardKeyType_Vector.h
--- a/Engine/Source/Runtime/AIModule/Classes/BehaviorTree/Blackboard/BlackboardKeyType_Vector.h
+++ b/Engine/Source/Runtime/AIModule/Classes/BehaviorTree/Blackboard/BlackboardKeyType_Vector.h
@@ -22,6 +22,13 @@ class UBlackboardKeyType_Vector : public UBlackboardKeyType
 	AIMODULE_API virtual EBlackboardCompare::Type CompareValues(const UBlackboardComponent& OwnerComp, const uint8* MemoryBlock,
 		const UBlackboardKeyType* OtherKeyOb, const uint8* OtherMemoryBlock) const override;
 
+	// If not set, will default to  FAISystem::InvalidLocation
+	UPROPERTY(EditDefaultsOnly, Category = Blackboard, meta = (EditCondition = "bUseDefaultValue"))
+	FVector DefaultValue = FVector::ZeroVector;
+
+	UPROPERTY()
+	bool bUseDefaultValue = false;
+
 protected:
 	AIMODULE_API virtual void InitializeMemory(UBlackboardComponent& OwnerComp, uint8* RawData) override;
 	AIMODULE_API virtual FString DescribeValue(const UBlackboardComponent& OwnerComp, const uint8* RawData) const override;
diff --git a/Engine/Source/Runtime/AIModule/Private/BehaviorTree/Blackboard/BlackboardKeyType_Bool.cpp b/Engine/Source/Runtime/AIModule/Private/BehaviorTree/Blackboard/BlackboardKeyType_Bool.cpp
--- a/Engine/Source/Runtime/AIModule/Private/BehaviorTree/Blackboard/BlackboardKeyType_Bool.cpp
+++ b/Engine/Source/Runtime/AIModule/Private/BehaviorTree/Blackboard/BlackboardKeyType_Bool.cpp
@@ -31,6 +31,11 @@ EBlackboardCompare::Type UBlackboardKeyType_Bool::CompareValues(const UBlackboar
 	return (MyValue == OtherValue) ? EBlackboardCompare::Equal : EBlackboardCompare::NotEqual;
 }
 
+void UBlackboardKeyType_Bool::InitializeMemory(UBlackboardComponent& OwnerComp, uint8* MemoryBlock)
+{
+	SetValue(this, MemoryBlock, bDefaultValue);
+}
+
 FString UBlackboardKeyType_Bool::DescribeValue(const UBlackboardComponent& OwnerComp, const uint8* RawData) const
 {
 	return GetValue(this, RawData) ? TEXT("true") : TEXT("false");
diff --git a/Engine/Source/Runtime/AIModule/Private/BehaviorTree/Blackboard/BlackboardKeyType_Float.cpp b/Engine/Source/Runtime/AIModule/Private/BehaviorTree/Blackboard/BlackboardKeyType_Float.cpp
--- a/Engine/Source/Runtime/AIModule/Private/BehaviorTree/Blackboard/BlackboardKeyType_Float.cpp
+++ b/Engine/Source/Runtime/AIModule/Private/BehaviorTree/Blackboard/BlackboardKeyType_Float.cpp
@@ -33,6 +33,11 @@ EBlackboardCompare::Type UBlackboardKeyType_Float::CompareValues(const UBlackboa
 		EBlackboardCompare::Less;
 }
 
+void UBlackboardKeyType_Float::InitializeMemory(UBlackboardComponent& OwnerComp, uint8* MemoryBlock)
+{
+	SetValue(this, MemoryBlock, DefaultValue);
+}
+
 FString UBlackboardKeyType_Float::DescribeValue(const UBlackboardComponent& OwnerComp, const uint8* RawData) const
 {
 	return FString::Printf(TEXT("%f"), GetValue(this, RawData));
diff --git a/Engine/Source/Runtime/AIModule/Private/BehaviorTree/Blackboard/BlackboardKeyType_Int.cpp b/Engine/Source/Runtime/AIModule/Private/BehaviorTree/Blackboard/BlackboardKeyType_Int.cpp
--- a/Engine/Source/Runtime/AIModule/Private/BehaviorTree/Blackboard/BlackboardKeyType_Int.cpp
+++ b/Engine/Source/Runtime/AIModule/Private/BehaviorTree/Blackboard/BlackboardKeyType_Int.cpp
@@ -33,6 +33,11 @@ EBlackboardCompare::Type UBlackboardKeyType_Int::CompareValues(const UBlackboard
 		EBlackboardCompare::Equal;
 }
 
+void UBlackboardKeyType_Int::InitializeMemory(UBlackboardComponent& OwnerComp, uint8* MemoryBlock)
+{
+	SetValue(this, MemoryBlock, DefaultValue);
+}
+
 FString UBlackboardKeyType_Int::DescribeValue(const UBlackboardComponent& OwnerComp, const uint8* RawData) const
 {
 	return FString::Printf(TEXT("%d"), GetValue(this, RawData));
diff --git a/Engine/Source/Runtime/AIModule/Private/BehaviorTree/Blackboard/BlackboardKeyType_Name.cpp b/Engine/Source/Runtime/AIModule/Private/BehaviorTree/Blackboard/BlackboardKeyType_Name.cpp
--- a/Engine/Source/Runtime/AIModule/Private/BehaviorTree/Blackboard/BlackboardKeyType_Name.cpp
+++ b/Engine/Source/Runtime/AIModule/Private/BehaviorTree/Blackboard/BlackboardKeyType_Name.cpp
@@ -31,6 +31,11 @@ EBlackboardCompare::Type UBlackboardKeyType_Name::CompareValues(const UBlackboar
 	return (MyValue == OtherValue) ? EBlackboardCompare::Equal : EBlackboardCompare::NotEqual;
 }
 
+void UBlackboardKeyType_Name::InitializeMemory(UBlackboardComponent& OwnerComp, uint8* MemoryBlock)
+{
+	SetValue(this, MemoryBlock, DefaultValue);
+}
+
 FString UBlackboardKeyType_Name::DescribeValue(const UBlackboardComponent& OwnerComp, const uint8* RawData) const
 {
 	return GetValue(this, RawData).ToString();
diff --git a/Engine/Source/Runtime/AIModule/Private/BehaviorTree/Blackboard/BlackboardKeyType_Rotator.cpp b/Engine/Source/Runtime/AIModule/Private/BehaviorTree/Blackboard/BlackboardKeyType_Rotator.cpp
--- a/Engine/Source/Runtime/AIModule/Private/BehaviorTree/Blackboard/BlackboardKeyType_Rotator.cpp
+++ b/Engine/Source/Runtime/AIModule/Private/BehaviorTree/Blackboard/BlackboardKeyType_Rotator.cpp
@@ -57,7 +57,14 @@ bool UBlackboardKeyType_Rotator::GetRotation(const UBlackboardComponent& OwnerCo
 
 void UBlackboardKeyType_Rotator::InitializeMemory(UBlackboardComponent& OwnerComp, uint8* RawData)
 {
-	SetValue(this, RawData, FAISystem::InvalidRotation);
+	if (bUseDefaultValue)
+	{
+		SetValue(this, RawData, DefaultValue);
+	}
+	else
+	{
+		SetValue(this,RawData, InvalidValue);
+	}
 }
 
 bool UBlackboardKeyType_Rotator::TestBasicOperation(const UBlackboardComponent& OwnerComp, const uint8* MemoryBlock, EBasicKeyOperation::Type Op) const
diff --git a/Engine/Source/Runtime/AIModule/Private/BehaviorTree/Blackboard/BlackboardKeyType_String.cpp b/Engine/Source/Runtime/AIModule/Private/BehaviorTree/Blackboard/BlackboardKeyType_String.cpp
--- a/Engine/Source/Runtime/AIModule/Private/BehaviorTree/Blackboard/BlackboardKeyType_String.cpp
+++ b/Engine/Source/Runtime/AIModule/Private/BehaviorTree/Blackboard/BlackboardKeyType_String.cpp
@@ -40,6 +40,11 @@ void UBlackboardKeyType_String::CopyValues(UBlackboardComponent& OwnerComp, uint
 	StringValue = ((UBlackboardKeyType_String*)SourceKeyOb)->StringValue;
 }
 
+void UBlackboardKeyType_String::InitializeMemory(UBlackboardComponent& OwnerComp, uint8* MemoryBlock)
+{
+	SetValue(this, MemoryBlock, DefaultValue);
+}
+
 FString UBlackboardKeyType_String::DescribeValue(const UBlackboardComponent& OwnerComp, const uint8* RawData) const
 {
 	return StringValue;
diff --git a/Engine/Source/Runtime/AIModule/Private/BehaviorTree/Blackboard/BlackboardKeyType_Vector.cpp b/Engine/Source/Runtime/AIModule/Private/BehaviorTree/Blackboard/BlackboardKeyType_Vector.cpp
--- a/Engine/Source/Runtime/AIModule/Private/BehaviorTree/Blackboard/BlackboardKeyType_Vector.cpp
+++ b/Engine/Source/Runtime/AIModule/Private/BehaviorTree/Blackboard/BlackboardKeyType_Vector.cpp
@@ -57,7 +57,14 @@ bool UBlackboardKeyType_Vector::GetLocation(const UBlackboardComponent& OwnerCom
 
 void UBlackboardKeyType_Vector::InitializeMemory(UBlackboardComponent& OwnerComp, uint8* RawData)
 {
-	SetValue(this, RawData, FAISystem::InvalidLocation);
+	if (bUseDefaultValue)
+	{
+		SetValue(this, RawData, DefaultValue);
+	}
+	else
+	{
+		SetValue(this, RawData, InvalidValue);
+	}
 }
 
 bool UBlackboardKeyType_Vector::TestBasicOperation(const UBlackboardComponent& OwnerComp, const uint8* MemoryBlock, EBasicKeyOperation::Type Op) const
