PR #11526: Stop BT in UGameplayBehavior_BehaviorTree on abort
PR #11582: Clean up GameplayBehavior_BehaviorTree on ctrl teardown

A cancelled behavior now explicitly stops the tree to not let it run until another different BT is pushed. It also solved the issue with repeated execution requests of the same BT asset that will fail as the component would think it's already running it.

Listen to possessed Pawn changes to prevent leaving the behavior running which could cause a resource leak and an ensure in OnTimerTick.

diff --git a/Engine/Plugins/Experimental/GameplayBehaviors/Source/GameplayBehaviorsModule/Private/AI/GameplayBehavior_BehaviorTree.cpp b/Engine/Plugins/Experimental/GameplayBehaviors/Source/GameplayBehaviorsModule/Private/AI/GameplayBehavior_BehaviorTree.cpp
--- a/Engine/Plugins/Experimental/GameplayBehaviors/Source/GameplayBehaviorsModule/Private/AI/GameplayBehavior_BehaviorTree.cpp
+++ b/Engine/Plugins/Experimental/GameplayBehaviors/Source/GameplayBehaviorsModule/Private/AI/GameplayBehavior_BehaviorTree.cpp
@@ -80,6 +80,8 @@ bool UGameplayBehavior_BehaviorTree::Trigger(AActor& InAvatar, const UGameplayBe
 		BTComp->ResumeLogic(TEXT("Allow inner BT to run within GameplayBehavior_BehaviorTree"));
 	}
 
+	AIController->OnPossessedPawnChanged.AddDynamic(this, &ThisClass::OnPossessedPawnChanged);
+
 	return true;
 }
 
@@ -110,13 +112,34 @@ void UGameplayBehavior_BehaviorTree::OnTimerTick()
 	}
 }
 
+void UGameplayBehavior_BehaviorTree::OnPossessedPawnChanged(APawn* OldPawn, APawn*)
+{
+	if (OldPawn != nullptr)
+	{
+		EndBehavior(*OldPawn, /*bInterrupted*/true);
+	}
+}
+
 void UGameplayBehavior_BehaviorTree::EndBehavior(AActor& InAvatar, const bool bInterrupted)
 {
 	Super::EndBehavior(InAvatar, bInterrupted);
 
-	if (PreviousBT && AIController)
+	if (AIController)
 	{
-		AIController->RunBehaviorTree(PreviousBT);
+		AIController->OnPossessedPawnChanged.RemoveAll(this);
+
+		if (bInterrupted)
+		{
+			if (UBehaviorTreeComponent* BTComp = Cast<UBehaviorTreeComponent>(AIController->GetBrainComponent()))
+			{
+				BTComp->StopTree(EBTStopMode::Safe);
+			}
+		}
+
+		if (PreviousBT)
+		{
+			AIController->RunBehaviorTree(PreviousBT);
+		}
 	}
 }
 
diff --git a/Engine/Plugins/Experimental/GameplayBehaviors/Source/GameplayBehaviorsModule/Public/AI/GameplayBehavior_BehaviorTree.h b/Engine/Plugins/Experimental/GameplayBehaviors/Source/GameplayBehaviorsModule/Public/AI/GameplayBehavior_BehaviorTree.h
--- a/Engine/Plugins/Experimental/GameplayBehaviors/Source/GameplayBehaviorsModule/Public/AI/GameplayBehavior_BehaviorTree.h
+++ b/Engine/Plugins/Experimental/GameplayBehaviors/Source/GameplayBehaviorsModule/Public/AI/GameplayBehavior_BehaviorTree.h
@@ -24,6 +24,9 @@ protected:
 
 	void OnTimerTick();
 
+	UFUNCTION()
+	void OnPossessedPawnChanged(APawn* OldPawn, APawn* NewPawn);
+
 	UPROPERTY()
 	TObjectPtr<UBehaviorTree> PreviousBT;
 
