Remove invalid pragma once to fix CIS
#rnx

diff --git a/Engine/Plugins/Experimental/GameplayBehaviors/Source/GameplayBehaviorsModule/Private/AI/ValueOrBBKey_GameplayTag.cpp b/Engine/Plugins/Experimental/GameplayBehaviors/Source/GameplayBehaviorsModule/Private/AI/ValueOrBBKey_GameplayTag.cpp
--- a/Engine/Plugins/Experimental/GameplayBehaviors/Source/GameplayBehaviorsModule/Private/AI/ValueOrBBKey_GameplayTag.cpp
+++ b/Engine/Plugins/Experimental/GameplayBehaviors/Source/GameplayBehaviorsModule/Private/AI/ValueOrBBKey_GameplayTag.cpp
@@ -1,6 +1,4 @@
 // Copyright Epic Games, Inc. All Rights Reserved.
-#pragma once
-
 #include "AI/ValueOrBBKey_GameplayTag.h"
 
 #include "BlackboardKeyType_GameplayTag.h"
