Node memory cleanup
Fetch the node memory from the behavior tree component instead of passing it to delegates

#Synced-CL 52165426

#Synced-CL 52165426

diff --git a/Engine/Plugins/Experimental/GameplayBehaviors/Source/GameplayBehaviorsModule/Private/AI/BTDecorator_GameplayTagQuery.cpp b/Engine/Plugins/Experimental/GameplayBehaviors/Source/GameplayBehaviorsModule/Private/AI/BTDecorator_GameplayTagQuery.cpp
--- a/Engine/Plugins/Experimental/GameplayBehaviors/Source/GameplayBehaviorsModule/Private/AI/BTDecorator_GameplayTagQuery.cpp
+++ b/Engine/Plugins/Experimental/GameplayBehaviors/Source/GameplayBehaviorsModule/Private/AI/BTDecorator_GameplayTagQuery.cpp
@@ -51,7 +51,7 @@ bool UBTDecorator_GameplayTagQuery::CalculateRawConditionValue(UBehaviorTreeComp
 	return GameplayTagQuery.Matches(SelectedActorTags);
 }
 
-void UBTDecorator_GameplayTagQuery::OnGameplayTagInQueryChanged(const FGameplayTag InTag, int32 NewCount, TWeakObjectPtr<UBehaviorTreeComponent> BehaviorTreeComponent, uint8* NodeMemory)
+void UBTDecorator_GameplayTagQuery::OnGameplayTagInQueryChanged(const FGameplayTag InTag, int32 NewCount, TWeakObjectPtr<UBehaviorTreeComponent> BehaviorTreeComponent)
 {
 	if (!BehaviorTreeComponent.IsValid())
 	{
@@ -101,7 +101,7 @@ void UBTDecorator_GameplayTagQuery::OnBecomeRelevant(UBehaviorTreeComponent& Own
 	{
 		for (const FGameplayTag& CurrentTag : QueryTags)
 		{
-			FDelegateHandle GameplayTagEventCallbackDelegate = MyMemory->CachedAbilitySystemComponent.Get()->RegisterGameplayTagEvent(CurrentTag, EGameplayTagEventType::Type::AnyCountChange).AddUObject(this, &UBTDecorator_GameplayTagQuery::OnGameplayTagInQueryChanged, TWeakObjectPtr<UBehaviorTreeComponent>(&OwnerComp), NodeMemory);
+			FDelegateHandle GameplayTagEventCallbackDelegate = MyMemory->CachedAbilitySystemComponent.Get()->RegisterGameplayTagEvent(CurrentTag, EGameplayTagEventType::Type::AnyCountChange).AddUObject(this, &UBTDecorator_GameplayTagQuery::OnGameplayTagInQueryChanged, TWeakObjectPtr<UBehaviorTreeComponent>(&OwnerComp));
 			MyMemory->GameplayTagEventHandles.Emplace(CurrentTag, GameplayTagEventCallbackDelegate);
 		}
 	}
diff --git a/Engine/Plugins/Experimental/GameplayBehaviors/Source/GameplayBehaviorsModule/Public/AI/BTDecorator_GameplayTagQuery.h b/Engine/Plugins/Experimental/GameplayBehaviors/Source/GameplayBehaviorsModule/Public/AI/BTDecorator_GameplayTagQuery.h
--- a/Engine/Plugins/Experimental/GameplayBehaviors/Source/GameplayBehaviorsModule/Public/AI/BTDecorator_GameplayTagQuery.h
+++ b/Engine/Plugins/Experimental/GameplayBehaviors/Source/GameplayBehaviorsModule/Public/AI/BTDecorator_GameplayTagQuery.h
@@ -26,7 +26,7 @@ class UBTDecorator_GameplayTagQuery : public UBTDecorator
 	UE_API virtual bool CalculateRawConditionValue(UBehaviorTreeComponent& OwnerComp, uint8* NodeMemory) const override;
 
 	/** Callback for when a tag in our query changes */
-	void OnGameplayTagInQueryChanged(const FGameplayTag InTag, int32 NewCount, TWeakObjectPtr<UBehaviorTreeComponent> BehaviorTreeComponent, uint8* NodeMemory);
+	void OnGameplayTagInQueryChanged(const FGameplayTag InTag, int32 NewCount, TWeakObjectPtr<UBehaviorTreeComponent> BehaviorTreeComponent);
 
 	UE_API virtual FString GetStaticDescription() const override;
 
